fix compilation error

diff --git a/Engine/Plugins/Experimental/Phya/Source/Phya/Classes/PhyaCollisionHandler.h b/Engine/Plugins/Experimental/Phya/Source/Phya/Classes/PhyaCollisionHandler.h
--- a/Engine/Plugins/Experimental/Phya/Source/Phya/Classes/PhyaCollisionHandler.h
+++ b/Engine/Plugins/Experimental/Phya/Source/Phya/Classes/PhyaCollisionHandler.h
@@ -43,7 +43,7 @@ class UPhyaCollisionHandler : public UPhysicsCollisionHandler
 	UAudioComponent* AudioComp;
 
 	UPROPERTY()
-	USoundWaveProcedural* ProceduralWave;
+	class USoundWaveProcedural* ProceduralWave;
 
 	TMap< FPhyaBodyInstancePair, TSharedPtr<FPhyaPairInfo> >	PairHash;
 
