Proper fixes for PhyaCollision plugin compile errors

diff --git a/Engine/Plugins/Experimental/Phya/Source/Phya/Classes/PhyaCollisionHandler.h b/Engine/Plugins/Experimental/Phya/Source/Phya/Classes/PhyaCollisionHandler.h
--- a/Engine/Plugins/Experimental/Phya/Source/Phya/Classes/PhyaCollisionHandler.h
+++ b/Engine/Plugins/Experimental/Phya/Source/Phya/Classes/PhyaCollisionHandler.h
@@ -4,6 +4,8 @@
 
 #include "PhyaCollisionHandler.generated.h"
 
+class USoundWaveProcedural;
+
 struct FPhyaBodyInstancePair
 {
 	FBodyInstance* Body0; // lower pointer
@@ -43,7 +45,7 @@ class UPhyaCollisionHandler : public UPhysicsCollisionHandler
 	UAudioComponent* AudioComp;
 
 	UPROPERTY()
-	class USoundWaveProcedural* ProceduralWave;
+	USoundWaveProcedural* ProceduralWave;
 
 	TMap< FPhyaBodyInstancePair, TSharedPtr<FPhyaPairInfo> >	PairHash;
 
diff --git a/Engine/Plugins/Experimental/Phya/Source/Phya/Private/PhyaCollisionHandler.cpp b/Engine/Plugins/Experimental/Phya/Source/Phya/Private/PhyaCollisionHandler.cpp
--- a/Engine/Plugins/Experimental/Phya/Source/Phya/Private/PhyaCollisionHandler.cpp
+++ b/Engine/Plugins/Experimental/Phya/Source/Phya/Private/PhyaCollisionHandler.cpp
@@ -3,6 +3,7 @@
 #include "PhyaPluginPrivatePCH.h"
 #include "PhysicsPublic.h"
 #include "Phya.hpp"
+#include "Sound/SoundWaveProcedural.h"
 
 #define N_BODIES 5
 
